ElectraPlayer: Enabled audio bitrate switching for audio-only DASH VoD streams

#rnx

diff --git a/Engine/Plugins/Media/ElectraPlayer/Source/ElectraPlayerRuntime/Private/Runtime/Decoder/AudioDecoder.cpp b/Engine/Plugins/Media/ElectraPlayer/Source/ElectraPlayerRuntime/Private/Runtime/Decoder/AudioDecoder.cpp
--- a/Engine/Plugins/Media/ElectraPlayer/Source/ElectraPlayerRuntime/Private/Runtime/Decoder/AudioDecoder.cpp
+++ b/Engine/Plugins/Media/ElectraPlayer/Source/ElectraPlayerRuntime/Private/Runtime/Decoder/AudioDecoder.cpp
@@ -135,6 +135,7 @@ private:
 
 	EAUChangeFlags GetAndPrepareInputAU();
 	void PrepareAU(TSharedPtrTS<FDecoderInput> AU);
+	IElectraDecoder::ECSDCompatibility IsCompatibleWith();
 
 	bool PostError(int32 ApiReturnValue, const FString& Message, uint16 Code, UEMediaError Error = UEMEDIA_ERROR_OK);
 	bool PostError(const IElectraDecoder::FError& InDecoderError);
@@ -649,6 +650,23 @@ FAudioDecoderImpl::EAUChangeFlags FAudioDecoderImpl::GetAndPrepareInputAU()
 	return NewAUFlags;
 }
 
+IElectraDecoder::ECSDCompatibility FAudioDecoderImpl::IsCompatibleWith()
+{
+	IElectraDecoder::ECSDCompatibility Compatibility = IElectraDecoder::ECSDCompatibility::Compatible;
+	if (DecoderInstance.IsValid() && CurrentAccessUnit.IsValid())
+	{
+		TMap<FString, FVariant> CSDOptions;
+		if (CurrentAccessUnit->AccessUnit->AUCodecData.IsValid())
+		{
+			CSDOptions.Emplace(TEXT("csd"), FVariant(CurrentAccessUnit->AccessUnit->AUCodecData->CodecSpecificData));
+			CSDOptions.Emplace(TEXT("dcr"), FVariant(CurrentAccessUnit->AccessUnit->AUCodecData->RawCSD));
+			CurrentAccessUnit->AccessUnit->AUCodecData->ParsedInfo.GetExtras().ConvertTo(CSDOptions, TEXT("$"));
+			Compatibility = DecoderInstance->IsCompatibleWith(CSDOptions);
+		}
+	}
+	return Compatibility;
+}
+
 IElectraDecoder::EOutputStatus FAudioDecoderImpl::HandleOutput()
 {
 	IElectraDecoder::EOutputStatus OutputStatus = IElectraDecoder::EOutputStatus::Available;
@@ -878,6 +896,7 @@ bool FAudioDecoderImpl::HandleDecoding()
 			if (CurrentCodecData.IsValid())
 			{
 				CSDOptions.Emplace(TEXT("csd"), FVariant(CurrentCodecData->CodecSpecificData));
+				CurrentCodecData->ParsedInfo.GetExtras().ConvertTo(CSDOptions, TEXT("$"));
 			}
 
 			IElectraDecoder::EDecoderError DecErr = DecoderInstance->DecodeAccessUnit(DecAU, CSDOptions);
@@ -1102,7 +1121,15 @@ void FAudioDecoderImpl::WorkerThread()
 			// Did the codec specific data change?
 			if ((NewAUFlags & EAUChangeFlags::CSDChanged) != EAUChangeFlags::None)
 			{
-				StartDraining(bIsAdaptiveDecoder ? EDecodingState::NormalDecoding : EDecodingState::NeedsReset);
+				// If the decoder is not adaptive, ask it how we have to handle the change.
+				if (!bIsAdaptiveDecoder)
+				{
+					IElectraDecoder::ECSDCompatibility Compatibility = IsCompatibleWith();
+					if (Compatibility == IElectraDecoder::ECSDCompatibility::Drain || Compatibility == IElectraDecoder::ECSDCompatibility::DrainAndReset)
+					{
+						StartDraining(Compatibility == IElectraDecoder::ECSDCompatibility::Drain ? EDecodingState::NormalDecoding : EDecodingState::NeedsReset);
+					}
+				}
 			}
 			// Is there a discontinuity that requires us to drain the decoder, including a switch to dummy-decoding?
 			else if ((NewAUFlags & EAUChangeFlags::Discontinuity) != EAUChangeFlags::None)
diff --git a/Engine/Plugins/Media/ElectraPlayer/Source/ElectraPlayerRuntime/Private/Runtime/Demuxer/ParserISO14496-12.cpp b/Engine/Plugins/Media/ElectraPlayer/Source/ElectraPlayerRuntime/Private/Runtime/Demuxer/ParserISO14496-12.cpp
--- a/Engine/Plugins/Media/ElectraPlayer/Source/ElectraPlayerRuntime/Private/Runtime/Demuxer/ParserISO14496-12.cpp
+++ b/Engine/Plugins/Media/ElectraPlayer/Source/ElectraPlayerRuntime/Private/Runtime/Demuxer/ParserISO14496-12.cpp
@@ -6,6 +6,7 @@
 
 #include "Demuxer/ParserISO14496-12.h"
 
+#include "Utilities/Utilities.h"
 #include "Utilities/UtilsMPEG.h"
 #include "Utils/MPEG/ElectraUtilsMPEGAudio.h"
 #include "Utils/MPEG/ElectraUtilsMPEGVideo.h"
@@ -5846,10 +5847,10 @@ private:
 			FString GetNameFromHandler() const override;
 			FTimeFraction GetDuration() const override;
 			ITrackIterator* CreateIterator() const override;
-			const TArray<uint8>& GetCodecSpecificData() const override;
-			const TArray<uint8>& GetCodecSpecificDataRAW() const override;
-			const FStreamCodecInformation& GetCodecInformation() const override;
-			const FBitrateInfo& GetBitrateInfo() const override;
+			const TArray<uint8> GetCodecSpecificData() const override;
+			const TArray<uint8> GetCodecSpecificDataRAW() const override;
+			const FStreamCodecInformation GetCodecInformation() const override;
+			const FBitrateInfo GetBitrateInfo() const override;
 			const FString GetLanguage() const override;
 			bool GetEncryptionInfo(FEncryptionInfo& OutEncryptionInfo) const override;
 			void GetPSSHBoxes(TArray<TArray<uint8>>& OutBoxes, bool bFromMOOV, bool bFromMOOF) const override;
@@ -6088,7 +6089,7 @@ private:
 		return TrackIterator;
 	}
 
-	const TArray<uint8>& FParserISO14496_12::FTrack::GetCodecSpecificData() const
+	const TArray<uint8> FParserISO14496_12::FTrack::GetCodecSpecificData() const
 	{
 		static TArray<uint8> EmptyCSD;
 		switch(CodecInformation.GetCodec())
@@ -6110,6 +6111,18 @@ private:
 			{
 				return CodecSpecificDataRAW;
 			}
+			case FStreamCodecInformation::ECodec::Audio4CC:
+			{
+				if (CodecInformation.GetCodec4CC() == Utils::Make4CC('O','p','u','s'))
+				{
+					FVariantValue dOps = CodecInformation.GetExtras().GetValue(TEXT("dOps_box"));
+					if (dOps.IsType(FVariantValue::EDataType::TypeU8Array))
+					{
+						return dOps.GetArray();
+					}
+				}
+				break;
+			}
 			default:
 			{
 				break;
@@ -6118,7 +6131,7 @@ private:
 		return EmptyCSD;
 	}
 
-	const TArray<uint8>& FParserISO14496_12::FTrack::GetCodecSpecificDataRAW() const
+	const TArray<uint8> FParserISO14496_12::FTrack::GetCodecSpecificDataRAW() const
 	{
 		static TArray<uint8> EmptyCSD;
 		switch(CodecInformation.GetCodec())
@@ -6149,12 +6162,12 @@ private:
 	}
 
 
-	const FStreamCodecInformation& FParserISO14496_12::FTrack::GetCodecInformation() const
+	const FStreamCodecInformation FParserISO14496_12::FTrack::GetCodecInformation() const
 	{
 		return CodecInformation;
 	}
 
-	const FParserISO14496_12::ITrack::FBitrateInfo& FParserISO14496_12::FTrack::GetBitrateInfo() const
+	const FParserISO14496_12::ITrack::FBitrateInfo FParserISO14496_12::FTrack::GetBitrateInfo() const
 	{
 		return BitrateInfo;
 	}
diff --git a/Engine/Plugins/Media/ElectraPlayer/Source/ElectraPlayerRuntime/Private/Runtime/Demuxer/ParserISO14496-12.h b/Engine/Plugins/Media/ElectraPlayer/Source/ElectraPlayerRuntime/Private/Runtime/Demuxer/ParserISO14496-12.h
--- a/Engine/Plugins/Media/ElectraPlayer/Source/ElectraPlayerRuntime/Private/Runtime/Demuxer/ParserISO14496-12.h
+++ b/Engine/Plugins/Media/ElectraPlayer/Source/ElectraPlayerRuntime/Private/Runtime/Demuxer/ParserISO14496-12.h
@@ -261,10 +261,10 @@ namespace Electra
 
 			virtual ITrackIterator* CreateIterator() const = 0;
 
-			virtual const TArray<uint8>& GetCodecSpecificData() const = 0;
-			virtual const TArray<uint8>& GetCodecSpecificDataRAW() const = 0;
-			virtual const FStreamCodecInformation& GetCodecInformation() const = 0;
-			virtual const FBitrateInfo& GetBitrateInfo() const = 0;
+			virtual const TArray<uint8> GetCodecSpecificData() const = 0;
+			virtual const TArray<uint8> GetCodecSpecificDataRAW() const = 0;
+			virtual const FStreamCodecInformation GetCodecInformation() const = 0;
+			virtual const FBitrateInfo GetBitrateInfo() const = 0;
 			virtual const FString GetLanguage() const = 0;
 			virtual bool GetEncryptionInfo(FEncryptionInfo& OutEncryptionInfo) const = 0;
 			virtual void GetPSSHBoxes(TArray<TArray<uint8>>& OutBoxes, bool bFromMOOV, bool bFromMOOF) const = 0;
diff --git a/Engine/Plugins/Media/ElectraPlayer/Source/ElectraPlayerRuntime/Private/Runtime/Demuxer/ParserMKV.cpp b/Engine/Plugins/Media/ElectraPlayer/Source/ElectraPlayerRuntime/Private/Runtime/Demuxer/ParserMKV.cpp
--- a/Engine/Plugins/Media/ElectraPlayer/Source/ElectraPlayerRuntime/Private/Runtime/Demuxer/ParserMKV.cpp
+++ b/Engine/Plugins/Media/ElectraPlayer/Source/ElectraPlayerRuntime/Private/Runtime/Demuxer/ParserMKV.cpp
@@ -2750,6 +2750,7 @@ namespace Electra
 			OutCodecInformation.SetMimeType(TEXT("audio/mp4"));
 			OutCodecInformation.SetCodec(FStreamCodecInformation::ECodec::Audio4CC);
 			OutCodecInformation.SetCodec4CC(Utils::Make4CC('O','p','u','s'));
+			OutCodecInformation.SetCodecSpecificData(OpusHead);
 			OutCodecInformation.SetStreamLanguageCode(InFromTrack->GetLanguage());
 			OutCodecInformation.SetCodecSpecifierRFC6381(TEXT("Opus"));
 			OutCodecInformation.SetSamplingRate(Audio->GetOutputSampleRate());
diff --git a/Engine/Plugins/Media/ElectraPlayer/Source/ElectraPlayerRuntime/Private/Runtime/Player/ABRRules/ABRVoDPlus.cpp b/Engine/Plugins/Media/ElectraPlayer/Source/ElectraPlayerRuntime/Private/Runtime/Player/ABRRules/ABRVoDPlus.cpp
--- a/Engine/Plugins/Media/ElectraPlayer/Source/ElectraPlayerRuntime/Private/Runtime/Player/ABRRules/ABRVoDPlus.cpp
+++ b/Engine/Plugins/Media/ElectraPlayer/Source/ElectraPlayerRuntime/Private/Runtime/Player/ABRRules/ABRVoDPlus.cpp
@@ -182,6 +182,16 @@ private:
 
 	void SetRenderRateScale(double InNewRate);
 
+	bool IsVideoOnly() const
+	{
+		return AudioWorkVars.NumQualities == 0;
+	}
+
+	bool IsAudioOnly() const
+	{
+		return VideoWorkVars.NumQualities == 0;
+	}
+
 	IABRInfoInterface* Info;
 	FCriticalSection Lock;
 	EMediaFormatType FormatType;
@@ -489,7 +499,7 @@ void FABROnDemandPlus::ReportDownloadEnd(const Metrics::FSegmentDownloadStats& S
 		}
 
 		// Adjust any forced bitrate duration.
-		if (SegmentDownloadStats.StreamType == EStreamType::Video && ForcedInitialBandwidthUntilSecondsBuffered > 0.0)
+		if (ForcedInitialBandwidthUntilSecondsBuffered > 0.0 && (SegmentDownloadStats.StreamType == EStreamType::Video || IsAudioOnly()))
 		{
 			ForcedInitialBandwidthUntilSecondsBuffered -= SegmentDownloadStats.DurationDownloaded;
 		}
@@ -775,27 +785,46 @@ IAdaptiveStreamSelector::ESegmentAction FABROnDemandPlus::EvaluateForQuality(TAr
 		set limits, like maximum bitrate and resolution.
 		The calling ABR will further reduce this list by removing streams that are not healthy.
 	*/
-	if (StreamType != EStreamType::Video)
-	{
-		InOutCandidates = Info->GetStreamInformations(StreamType);
-		return IAdaptiveStreamSelector::ESegmentAction::FetchNext;
-	}
-
 	if (InOutCandidates.Num())
 	{
-		TArray<TSharedPtrTS<FABRStreamInformation>> PossibleRepresentations;
-		PossibleRepresentations.Add(InOutCandidates[0]);
-		const int32 MaxAllowedBandwidth = Info->GetBandwidthCeiling();
-		const FStreamCodecInformation::FResolution MaxAllowedResolution = Info->GetMaxStreamResolution();
-		for(int32 nStr=1; nStr<InOutCandidates.Num(); ++nStr)
+		if (StreamType == EStreamType::Video)
+		{
+			TArray<TSharedPtrTS<FABRStreamInformation>> PossibleRepresentations;
+			PossibleRepresentations.Add(InOutCandidates[0]);
+			const int32 MaxAllowedBandwidth = Info->GetBandwidthCeiling();
+			const FStreamCodecInformation::FResolution MaxAllowedResolution = Info->GetMaxStreamResolution();
+			for(int32 nStr=1; nStr<InOutCandidates.Num(); ++nStr)
+			{
+				// Check if bitrate and resolution are acceptable
+				if (InOutCandidates[nStr]->Bitrate <= MaxAllowedBandwidth && !InOutCandidates[nStr]->Resolution.ExceedsLimit(MaxAllowedResolution))
+				{
+					PossibleRepresentations.Add(InOutCandidates[nStr]);
+				}
+			}
+			Swap(InOutCandidates, PossibleRepresentations);
+		}
+		else if (StreamType == EStreamType::Audio)
 		{
-			// Check if bitrate and resolution are acceptable
-			if (InOutCandidates[nStr]->Bitrate <= MaxAllowedBandwidth && !InOutCandidates[nStr]->Resolution.ExceedsLimit(MaxAllowedResolution))
+			if (IsAudioOnly())
+			{
+				TArray<TSharedPtrTS<FABRStreamInformation>> PossibleRepresentations;
+				PossibleRepresentations.Add(InOutCandidates[0]);
+				const int32 MaxAllowedBandwidth = Info->GetBandwidthCeiling();
+				for(int32 nStr=1; nStr<InOutCandidates.Num(); ++nStr)
+				{
+					// Check if bitrate is acceptable
+					if (InOutCandidates[nStr]->Bitrate <= MaxAllowedBandwidth)
+					{
+						PossibleRepresentations.Add(InOutCandidates[nStr]);
+					}
+				}
+				Swap(InOutCandidates, PossibleRepresentations);
+			}
+			else
 			{
-				PossibleRepresentations.Add(InOutCandidates[nStr]);
+				InOutCandidates = Info->GetStreamInformations(StreamType);
 			}
 		}
-		Swap(InOutCandidates, PossibleRepresentations);
 	}
 	return IAdaptiveStreamSelector::ESegmentAction::FetchNext;
 }
@@ -813,7 +842,7 @@ IAdaptiveStreamSelector::ESegmentAction FABROnDemandPlus::PerformSelection(const
 		const double AvailableDuration = GetPlayablePlayerDuration(bEOS, StreamType);
 		const double DownloadLatency = WorkVars ? WorkVars->AverageLatency.GetWeightedMax(DefaultNetworkLatency) : DefaultNetworkLatency;
 
-		if (WorkVars && StreamType == EStreamType::Video)
+		if (WorkVars && (StreamType == EStreamType::Video || IsAudioOnly()))
 		{
 			FScopeLock lock(&WorkVars->Lock);
 
@@ -1001,12 +1030,15 @@ IAdaptiveStreamSelector::EHandlingAction FABROnDemandPlus::PeriodicHandle()
 
 void FABROnDemandPlus::SetRenderRateScale(double InNewRate)
 {
-	FTimeRange pr = Info->ABRGetSupportedRenderRateScale();
-	if (pr.IsValid())
+	if (!IsAudioOnly())
 	{
-		double MinAllowed = pr.Start.GetAsSeconds(1.0);
-		double MaxAllowed = pr.End.GetAsSeconds(1.0);
-		Info->ABRSetRenderRateScale(Utils::Min(Utils::Max(MinAllowed, InNewRate), MaxAllowed));
+		FTimeRange pr = Info->ABRGetSupportedRenderRateScale();
+		if (pr.IsValid())
+		{
+			double MinAllowed = pr.Start.GetAsSeconds(1.0);
+			double MaxAllowed = pr.End.GetAsSeconds(1.0);
+			Info->ABRSetRenderRateScale(Utils::Min(Utils::Max(MinAllowed, InNewRate), MaxAllowed));
+		}
 	}
 }
 
diff --git a/Engine/Plugins/Media/ElectraPlayer/Source/ElectraPlayerRuntime/Private/Runtime/Player/DASH/ManifestDASH.cpp b/Engine/Plugins/Media/ElectraPlayer/Source/ElectraPlayerRuntime/Private/Runtime/Player/DASH/ManifestDASH.cpp
--- a/Engine/Plugins/Media/ElectraPlayer/Source/ElectraPlayerRuntime/Private/Runtime/Player/DASH/ManifestDASH.cpp
+++ b/Engine/Plugins/Media/ElectraPlayer/Source/ElectraPlayerRuntime/Private/Runtime/Player/DASH/ManifestDASH.cpp
@@ -336,7 +336,14 @@ TSharedPtrTS<IProducerReferenceTimeInfo> FManifestDASH::GetProducerReferenceTime
 TRangeSet<double> FManifestDASH::GetPossiblePlaybackRates(EPlayRateType InForType) const
 {
 	TRangeSet<double> Ranges;
-	Ranges.Add(TRange<double>{1.0}); // normal (real-time) playback rate
+	if (InForType == IManifest::EPlayRateType::UnthinnedRate)
+	{
+		Ranges.Add(TRange<double>::Inclusive(0.1, 4.0));
+	}
+	else
+	{
+		Ranges.Add(TRange<double>{1.0}); // normal (real-time) playback rate
+	}
 	Ranges.Add(TRange<double>{0.0}); // and pause
 	return Ranges;
 }
diff --git a/Engine/Plugins/Media/OpusDecoderElectra/Source/OpusDecoderElectra/Private/OpusDecoder/ElectraMediaOpusDecoder.cpp b/Engine/Plugins/Media/OpusDecoderElectra/Source/OpusDecoderElectra/Private/OpusDecoder/ElectraMediaOpusDecoder.cpp
--- a/Engine/Plugins/Media/OpusDecoderElectra/Source/OpusDecoderElectra/Private/OpusDecoder/ElectraMediaOpusDecoder.cpp
+++ b/Engine/Plugins/Media/OpusDecoderElectra/Source/OpusDecoderElectra/Private/OpusDecoder/ElectraMediaOpusDecoder.cpp
@@ -157,6 +157,39 @@ public:
 	{ }
 
 private:
+	struct FDOpsConfig
+	{
+		int32 NumberOfOutputChannels = 0;
+		int32 SampleRate = 0;
+		int32 PreSkip = 0;
+		int32 OutputGain = 0;
+		int32 ChannelMappingFamily = 0;
+		int32 StreamCount = 0;
+		int32 CoupledCount = 0;
+		uint8 ChannelMapping[256] = {255};
+		void Reset()
+		{
+			NumberOfOutputChannels = 0;
+			SampleRate = 0;
+			PreSkip = 0;
+			OutputGain = 0;
+			ChannelMappingFamily = 0;
+			StreamCount = 0;
+			CoupledCount = 0;
+			FMemory::Memset(ChannelMapping, 255);
+		}
+		bool SameAs(const FDOpsConfig& rhs)
+		{
+			return NumberOfOutputChannels == rhs.NumberOfOutputChannels &&
+				   SampleRate == rhs.SampleRate &&
+				   PreSkip == rhs.PreSkip &&
+				   OutputGain == rhs.OutputGain &&
+				   ChannelMappingFamily == rhs.ChannelMappingFamily &&
+				   StreamCount == rhs.StreamCount &&
+				   CoupledCount == rhs.CoupledCount &&
+				   FMemory::Memcmp(ChannelMapping, rhs.ChannelMapping, StreamCount) == 0;
+		}
+	};
 
 	static constexpr uint32 Make4CC(const uint8 A, const uint8 B, const uint8 C, const uint8 D)
 	{
@@ -169,7 +202,7 @@ private:
 		return 48000;
 	}
 
-	bool Parse_dOps();
+	bool Parse_dOps(FDOpsConfig& OutConfig, const TArray<uint8>& IndOpsBox, bool bFailOnError);
 
 	bool PostError(int32 ApiReturnValue, FString Message, int32 Code);
 
@@ -189,16 +222,8 @@ private:
 	bool bFlushPending = false;
 
 	// Input configuration
-	TArray<uint8> dOpsBox;
+	FDOpsConfig DOpsConfig;
 	bool bHaveParseddOps = false;
-	int32 CfgNumberOfOutputChannels = 0;
-	int32 CfgSampleRate = 0;
-	int32 CfgPreSkip = 0;
-	int32 CfgOutputGain = 0;
-	int32 CfgChannelMappingFamily = 0;
-	int32 CfgStreamCount = 0;
-	int32 CfgCoupledCount = 0;
-	uint8 CfgChannelMapping[256] = {255};
 	
 	// Output
 	TArray<IElectraDecoderAudioOutput::EChannelPosition> OutputChannelMap;
@@ -282,10 +307,6 @@ TSharedPtr<IElectraCodecFactory, ESPMode::ThreadSafe> FElectraMediaOpusDecoder::
 FElectraOpusDecoder::FElectraOpusDecoder(const TMap<FString, FVariant>& InOptions, TSharedPtr<IElectraDecoderResourceDelegate, ESPMode::ThreadSafe> InResourceDelegate)
 {
 	Codec4CC = (uint32)ElectraDecodersUtil::GetVariantValueSafeU64(InOptions, TEXT("codec_4cc"), 0);
-	if (InOptions.Contains(TEXT("$dOps_box")))
-	{
-		dOpsBox = ElectraDecodersUtil::GetVariantValueUInt8Array(InOptions, TEXT("$dOps_box"));
-	}
 }
 
 FElectraOpusDecoder::~FElectraOpusDecoder()
@@ -296,37 +317,37 @@ FElectraOpusDecoder::~FElectraOpusDecoder()
 	Close();
 }
 
-bool FElectraOpusDecoder::Parse_dOps()
+bool FElectraOpusDecoder::Parse_dOps(FDOpsConfig& OutConfig, const TArray<uint8>& IndOpsBox, bool bFailOnError)
 {
-	if (dOpsBox.Num() == 0)
+	if (IndOpsBox.Num() == 0)
 	{
-		return PostError(0, TEXT("There is no 'dOps' box to get Opus information from"), ERRCODE_INTERNAL_FAILED_TO_PARSE_CSD);
+		return !bFailOnError ? false : PostError(0, TEXT("There is no 'dOps' box to get Opus information from"), ERRCODE_INTERNAL_FAILED_TO_PARSE_CSD);
 	}
-	else if (dOpsBox.Num() < 11)
+	else if (IndOpsBox.Num() < 11)
 	{
-		return PostError(0, TEXT("Incomplete 'dOps' box"), ERRCODE_INTERNAL_FAILED_TO_PARSE_CSD);
+		return !bFailOnError ? false : PostError(0, TEXT("Incomplete 'dOps' box"), ERRCODE_INTERNAL_FAILED_TO_PARSE_CSD);
 	}
 
-	ElectraDecodersUtil::FMP4AtomReader rd(dOpsBox.GetData(), dOpsBox.Num());
+	ElectraDecodersUtil::FMP4AtomReader rd(IndOpsBox.GetData(), IndOpsBox.Num());
 	uint32 Value32;
 	uint16 Value16;
 	uint8 Value8;
 	rd.Read(Value8);		// Version
 	if (Value8 != 0)
 	{
-		return PostError(0, TEXT("Unsupported 'dOps' box version"), ERRCODE_INTERNAL_FAILED_TO_PARSE_CSD);
+		return !bFailOnError ? false : PostError(0, TEXT("Unsupported 'dOps' box version"), ERRCODE_INTERNAL_FAILED_TO_PARSE_CSD);
 	}
 	rd.Read(Value8);		// OutputChannelCount
 	rd.Read(Value16);		// PreSkip
 	rd.Read(Value32);		// InputSampleRate
-	CfgNumberOfOutputChannels = Value8;
-	CfgPreSkip = Value16;
-	CfgSampleRate = Value32;
+	OutConfig.NumberOfOutputChannels = Value8;
+	OutConfig.PreSkip = Value16;
+	OutConfig.SampleRate = Value32;
 	rd.Read(Value16);		// OutputGain
-	CfgOutputGain = Value16;
+	OutConfig.OutputGain = Value16;
 	rd.Read(Value8);		// ChannelMappingFamily
-	CfgChannelMappingFamily = Value8;
-	if (CfgChannelMappingFamily)
+	OutConfig.ChannelMappingFamily = Value8;
+	if (OutConfig.ChannelMappingFamily)
 	{
 		/*
 			Channel mapping family:
@@ -336,24 +357,24 @@ bool FElectraOpusDecoder::Parse_dOps()
 				3 : Ambisonics with demixing matrix
 				255 : discrete channels
 		*/
-		if (CfgChannelMappingFamily == 2 || CfgChannelMappingFamily == 3)
+		if (OutConfig.ChannelMappingFamily == 2 || OutConfig.ChannelMappingFamily == 3)
 		{
-			return PostError(0, TEXT("Ambisonics is not supported"), ERRCODE_INTERNAL_FAILED_TO_PARSE_CSD);
+			return !bFailOnError ? false : PostError(0, TEXT("Ambisonics is not supported"), ERRCODE_INTERNAL_FAILED_TO_PARSE_CSD);
 		}
-		else if (CfgChannelMappingFamily > 3 && CfgChannelMappingFamily < 255)
+		else if (OutConfig.ChannelMappingFamily > 3 && OutConfig.ChannelMappingFamily < 255)
 		{
-			return PostError(0, TEXT("Unsupported channel mapping family"), ERRCODE_INTERNAL_FAILED_TO_PARSE_CSD);
+			return !bFailOnError ? false : PostError(0, TEXT("Unsupported channel mapping family"), ERRCODE_INTERNAL_FAILED_TO_PARSE_CSD);
 		}
-		if (rd.GetNumBytesRemaining() < 2+CfgNumberOfOutputChannels)
+		if (rd.GetNumBytesRemaining() < 2+OutConfig.NumberOfOutputChannels)
 		{
-			return PostError(0, TEXT("Incomplete 'dOps' box"), ERRCODE_INTERNAL_FAILED_TO_PARSE_CSD);
+			return !bFailOnError ? false : PostError(0, TEXT("Incomplete 'dOps' box"), ERRCODE_INTERNAL_FAILED_TO_PARSE_CSD);
 		}
 		rd.Read(Value8);	// StreamCount
-		CfgStreamCount = Value8;
+		OutConfig.StreamCount = Value8;
 		rd.Read(Value8);	// CoupledCount
-		CfgCoupledCount = Value8;
+		OutConfig.CoupledCount = Value8;
 		// Pre-set channel mapping array to all disabled.
-		FMemory::Memset(CfgChannelMapping, 255);
+		FMemory::Memset(OutConfig.ChannelMapping, 255);
 		/*
 			The channel mapping array determines the order in which the input stream is mapped to the
 			output channels. The default order produces output according to the Vorbis order
@@ -368,7 +389,7 @@ bool FElectraOpusDecoder::Parse_dOps()
 			Coupled streams are stereo and thus have 2 channels, while uncoupled streams are single channel only.
 			The index values in this table are defined like this:
 				For j=0; j<OutputChannelCount:
-					Index = CfgChannelMapping[j];
+					Index = ChannelMapping[j];
 					If Index == 255:
 						OutputChannel[j] is made silent
 					Elseif Index < 2*CoupledCount:
@@ -379,19 +400,18 @@ bool FElectraOpusDecoder::Parse_dOps()
 			You do not necessarily know which stream contains which original input channel, which is why
 			the default mapping table exists and indicates the abovementioned Vorbis mapping.
 		*/
-		for(int32 i=0; i<CfgNumberOfOutputChannels; ++i)
+		for(int32 i=0; i<OutConfig.NumberOfOutputChannels; ++i)
 		{
-			rd.Read(CfgChannelMapping[i]);
+			rd.Read(OutConfig.ChannelMapping[i]);
 		}
 	}
 	else
 	{
-		CfgStreamCount = 1;
-		CfgCoupledCount = CfgNumberOfOutputChannels > 1 ? 1 : 0;
-		CfgChannelMapping[0] = 0;
-		CfgChannelMapping[1] = 1;
+		OutConfig.StreamCount = 1;
+		OutConfig.CoupledCount = OutConfig.NumberOfOutputChannels > 1 ? 1 : 0;
+		OutConfig.ChannelMapping[0] = 0;
+		OutConfig.ChannelMapping[1] = 1;
 	}
-	bHaveParseddOps = true;
 	return true;
 }
 
@@ -401,17 +421,18 @@ bool FElectraOpusDecoder::InternalDecoderCreate()
 {
 	if (!DecoderHandle)
 	{
-		int32 DecoderAllocSize = (int32) opus_multistream_decoder_get_size(CfgStreamCount, CfgCoupledCount);
+		int32 DecoderAllocSize = (int32) opus_multistream_decoder_get_size(DOpsConfig.StreamCount, DOpsConfig.CoupledCount);
 		if (DecoderAllocSize < 0)
 		{
 			return PostError(DecoderAllocSize, TEXT("opus_multistream_decoder_get_size() failed"), ERRCODE_INTERNAL_FAILED_TO_DECODE_INPUT);
 		}
 		DecoderHandle = (OpusMSDecoder*) FMemory::Malloc(DecoderAllocSize);
-		int32 Result = opus_multistream_decoder_init(DecoderHandle, OpusSamplingRate(), CfgNumberOfOutputChannels, CfgStreamCount, CfgCoupledCount, CfgChannelMapping);
+		int32 Result = opus_multistream_decoder_init(DecoderHandle, OpusSamplingRate(), DOpsConfig.NumberOfOutputChannels, DOpsConfig.StreamCount, DOpsConfig.CoupledCount, DOpsConfig.ChannelMapping);
 		if (Result != OPUS_OK)
 		{
 			return PostError(Result, TEXT("opus_multistream_decoder_init() failed"), ERRCODE_INTERNAL_FAILED_TO_DECODE_INPUT);
 		}
+		RemainingPreSkip = -1;
 	}
 	return true;
 }
@@ -447,14 +468,28 @@ bool FElectraOpusDecoder::PostError(int32 ApiReturnValue, FString Message, int32
 void FElectraOpusDecoder::Close()
 {
 	ResetToCleanStart();
-	InternalDecoderDestroy();
 	// Set the error state that all subsequent calls will fail.
 	PostError(0, TEXT("Already closed"), ERRCODE_INTERNAL_ALREADY_CLOSED);
 }
 
 IElectraDecoder::ECSDCompatibility FElectraOpusDecoder::IsCompatibleWith(const TMap<FString, FVariant>& CSDAndAdditionalOptions)
 {
-	return IElectraDecoder::ECSDCompatibility::Compatible;
+	// No configuration parsed yet, so this is deemed compatible.
+	if (!bHaveParseddOps)
+	{
+		return IElectraDecoder::ECSDCompatibility::Compatible;
+	}
+	TArray<uint8> SidebandData = ElectraDecodersUtil::GetVariantValueUInt8Array(CSDAndAdditionalOptions, TEXT("csd"));
+	FDOpsConfig cfg;
+	if (!Parse_dOps(cfg, SidebandData, false))
+	{
+		return IElectraDecoder::ECSDCompatibility::DrainAndReset;
+	}
+	if (cfg.SameAs(DOpsConfig))
+	{
+		return IElectraDecoder::ECSDCompatibility::Compatible;
+	}
+	return IElectraDecoder::ECSDCompatibility::DrainAndReset;
 }
 
 bool FElectraOpusDecoder::ResetToCleanStart()
@@ -464,14 +499,9 @@ bool FElectraOpusDecoder::ResetToCleanStart()
 	RemainingPreSkip = -1;
 
 	bHaveParseddOps = false;
-	CfgNumberOfOutputChannels = 0;
-	CfgSampleRate = 0;
-	CfgPreSkip = 0;
-	CfgOutputGain = 0;
-	CfgChannelMappingFamily = 0;
-	CfgStreamCount = 0;
-	CfgCoupledCount = 0;
+	DOpsConfig.Reset();
 	OutputChannelMap.Empty();
+	InternalDecoderDestroy();
 	return true;
 }
 
@@ -488,51 +518,51 @@ bool FElectraOpusDecoder::SetupChannelMap()
 	}
 	// Pre-init with all channels disabled.
 	OutputChannelMap.Empty();
-	OutputChannelMap.Init(IElectraDecoderAudioOutput::EChannelPosition::Disabled, CfgNumberOfOutputChannels);
+	OutputChannelMap.Init(IElectraDecoderAudioOutput::EChannelPosition::Disabled, DOpsConfig.NumberOfOutputChannels);
 
-	if (CfgChannelMappingFamily == 0)
+	if (DOpsConfig.ChannelMappingFamily == 0)
 	{
-		if (CfgNumberOfOutputChannels != 1 && CfgNumberOfOutputChannels != 2)
+		if (DOpsConfig.NumberOfOutputChannels != 1 && DOpsConfig.NumberOfOutputChannels != 2)
 		{
-			return PostError(0, FString::Printf(TEXT("Unsupported number of channels (%d) for mapping family 0"), CfgNumberOfOutputChannels), ERRCODE_INTERNAL_UNSUPPORTED_CHANNEL_LAYOUT);
+			return PostError(0, FString::Printf(TEXT("Unsupported number of channels (%d) for mapping family 0"), DOpsConfig.NumberOfOutputChannels), ERRCODE_INTERNAL_UNSUPPORTED_CHANNEL_LAYOUT);
 		}
-		if (CfgNumberOfOutputChannels == 1)
+		if (DOpsConfig.NumberOfOutputChannels == 1)
 		{
 			OutputChannelMap[0] = IElectraDecoderAudioOutput::EChannelPosition::C;
 		}
-		else if (CfgNumberOfOutputChannels == 2)
+		else if (DOpsConfig.NumberOfOutputChannels == 2)
 		{
 			OutputChannelMap[0] = IElectraDecoderAudioOutput::EChannelPosition::L;
 			OutputChannelMap[1] = IElectraDecoderAudioOutput::EChannelPosition::R;
 		}
 	}
-	else if (CfgChannelMappingFamily == 1)
+	else if (DOpsConfig.ChannelMappingFamily == 1)
 	{
 		// Channel mapping family 1 provides a default channel mapping table in Vorbis channel order.
 		// We use that order to remap the channels to our positions.
-		if (CfgNumberOfOutputChannels == 1)
+		if (DOpsConfig.NumberOfOutputChannels == 1)
 		{
 			OutputChannelMap[0] = IElectraDecoderAudioOutput::EChannelPosition::C;
 		}
-		else if (CfgNumberOfOutputChannels == 2)
+		else if (DOpsConfig.NumberOfOutputChannels == 2)
 		{
 			OutputChannelMap[0] = IElectraDecoderAudioOutput::EChannelPosition::L;
 			OutputChannelMap[1] = IElectraDecoderAudioOutput::EChannelPosition::R;
 		}
-		else if (CfgNumberOfOutputChannels == 3)
+		else if (DOpsConfig.NumberOfOutputChannels == 3)
 		{
 			OutputChannelMap[0] = IElectraDecoderAudioOutput::EChannelPosition::L;
 			OutputChannelMap[1] = IElectraDecoderAudioOutput::EChannelPosition::C;
 			OutputChannelMap[2] = IElectraDecoderAudioOutput::EChannelPosition::R;
 		}
-		else if (CfgNumberOfOutputChannels == 4)
+		else if (DOpsConfig.NumberOfOutputChannels == 4)
 		{
 			OutputChannelMap[0] = IElectraDecoderAudioOutput::EChannelPosition::L;
 			OutputChannelMap[1] = IElectraDecoderAudioOutput::EChannelPosition::R;
 			OutputChannelMap[2] = IElectraDecoderAudioOutput::EChannelPosition::Ls;
 			OutputChannelMap[3] = IElectraDecoderAudioOutput::EChannelPosition::Rs;
 		}
-		else if (CfgNumberOfOutputChannels == 5)
+		else if (DOpsConfig.NumberOfOutputChannels == 5)
 		{
 			OutputChannelMap[0] = IElectraDecoderAudioOutput::EChannelPosition::L;
 			OutputChannelMap[1] = IElectraDecoderAudioOutput::EChannelPosition::C;
@@ -540,7 +570,7 @@ bool FElectraOpusDecoder::SetupChannelMap()
 			OutputChannelMap[3] = IElectraDecoderAudioOutput::EChannelPosition::Ls;
 			OutputChannelMap[4] = IElectraDecoderAudioOutput::EChannelPosition::Rs;
 		}
-		else if (CfgNumberOfOutputChannels == 6)
+		else if (DOpsConfig.NumberOfOutputChannels == 6)
 		{
 			OutputChannelMap[0] = IElectraDecoderAudioOutput::EChannelPosition::L;
 			OutputChannelMap[1] = IElectraDecoderAudioOutput::EChannelPosition::C;
@@ -549,7 +579,7 @@ bool FElectraOpusDecoder::SetupChannelMap()
 			OutputChannelMap[4] = IElectraDecoderAudioOutput::EChannelPosition::Rs;
 			OutputChannelMap[5] = IElectraDecoderAudioOutput::EChannelPosition::LFE;
 		}
-		else if (CfgNumberOfOutputChannels == 7)
+		else if (DOpsConfig.NumberOfOutputChannels == 7)
 		{
 			OutputChannelMap[0] = IElectraDecoderAudioOutput::EChannelPosition::L;
 			OutputChannelMap[1] = IElectraDecoderAudioOutput::EChannelPosition::C;
@@ -559,7 +589,7 @@ bool FElectraOpusDecoder::SetupChannelMap()
 			OutputChannelMap[5] = IElectraDecoderAudioOutput::EChannelPosition::Cs;
 			OutputChannelMap[6] = IElectraDecoderAudioOutput::EChannelPosition::LFE;
 		}
-		else if (CfgNumberOfOutputChannels == 8)
+		else if (DOpsConfig.NumberOfOutputChannels == 8)
 		{
 			OutputChannelMap[0] = IElectraDecoderAudioOutput::EChannelPosition::L;
 			OutputChannelMap[1] = IElectraDecoderAudioOutput::EChannelPosition::C;
@@ -572,20 +602,20 @@ bool FElectraOpusDecoder::SetupChannelMap()
 		}
 		else
 		{
-			return PostError(0, FString::Printf(TEXT("Unsupported number of channels (%d) for mapping family 1"), CfgNumberOfOutputChannels), ERRCODE_INTERNAL_UNSUPPORTED_CHANNEL_LAYOUT);
+			return PostError(0, FString::Printf(TEXT("Unsupported number of channels (%d) for mapping family 1"), DOpsConfig.NumberOfOutputChannels), ERRCODE_INTERNAL_UNSUPPORTED_CHANNEL_LAYOUT);
 		}
 	}
-	else if (CfgChannelMappingFamily == 255)
+	else if (DOpsConfig.ChannelMappingFamily == 255)
 	{
 		// Unspecified order
-		for(int32 i=0; i<CfgNumberOfOutputChannels; ++i)
+		for(int32 i=0; i<DOpsConfig.NumberOfOutputChannels; ++i)
 		{
 			OutputChannelMap[i] = static_cast<IElectraDecoderAudioOutput::EChannelPosition>(static_cast<int32>(IElectraDecoderAudioOutput::EChannelPosition::Unspec0) + i);
 		}
 	}
 	else
 	{
-		return PostError(0, FString::Printf(TEXT("Unsupported channel mapping family (%d)"), CfgChannelMappingFamily), ERRCODE_INTERNAL_UNSUPPORTED_CHANNEL_LAYOUT);
+		return PostError(0, FString::Printf(TEXT("Unsupported channel mapping family (%d)"), DOpsConfig.ChannelMappingFamily), ERRCODE_INTERNAL_UNSUPPORTED_CHANNEL_LAYOUT);
 	}
 	return true;
 }
@@ -601,7 +631,7 @@ bool FElectraOpusDecoder::ProcessInput(const void* InData, int64 InDataSize)
 	// Set the pre-skip if necessary (first packet)
 	if (RemainingPreSkip < 0)
 	{
-		RemainingPreSkip = CfgPreSkip;
+		RemainingPreSkip = DOpsConfig.PreSkip;
 	}
 
 	int32 NumExpectedDecodedSamples = opus_packet_get_nb_samples(reinterpret_cast<const unsigned char*>(InData), (int32)InDataSize, OpusSamplingRate());
@@ -610,10 +640,10 @@ bool FElectraOpusDecoder::ProcessInput(const void* InData, int64 InDataSize)
 		return PostError(NumExpectedDecodedSamples, TEXT("opus_packet_get_nb_samples() failed"), ERRCODE_INTERNAL_FAILED_TO_DECODE_INPUT);
 	}
 
-	int32 AllocSize = sizeof(float) * CfgNumberOfOutputChannels * NumExpectedDecodedSamples;
+	int32 AllocSize = sizeof(float) * DOpsConfig.NumberOfOutputChannels * NumExpectedDecodedSamples;
 	CurrentOutput->Buffer = (float*)FMemory::Malloc(AllocSize);
 	CurrentOutput->ChannelPositions = OutputChannelMap;
-	CurrentOutput->NumChannels = CfgNumberOfOutputChannels;
+	CurrentOutput->NumChannels = DOpsConfig.NumberOfOutputChannels;
 	CurrentOutput->SampleRate = OpusSamplingRate();
 
 	int32 Result = opus_multistream_decode_float(DecoderHandle, reinterpret_cast<const unsigned char*>(InData), (int32)InDataSize, CurrentOutput->Buffer, NumExpectedDecodedSamples, 0);
@@ -623,15 +653,6 @@ bool FElectraOpusDecoder::ProcessInput(const void* InData, int64 InDataSize)
 	}
 	CurrentOutput->NumFrames = Result;
 
-#if 0
-	/*
-		We do not need to apply the pre-skip here right now. It must be included in an edit list which offsets
-		the PTS accordingly and the audio decoder will adjust the output accordingly.
-		
-		We leave the code here in case the audio decoder will not do this in which case it should be easy
-		to activate this bit via some configuration option.
-	*/
-
 	// Apply pre skip
 	if (RemainingPreSkip > 0)
 	{
@@ -653,7 +674,6 @@ bool FElectraOpusDecoder::ProcessInput(const void* InData, int64 InDataSize)
 			RemainingPreSkip = 0;
 		}
 	}
-#endif
 	return true;
 }
 
@@ -682,9 +702,13 @@ IElectraDecoder::EDecoderError FElectraOpusDecoder::DecodeAccessUnit(const FInpu
 	if (InInputAccessUnit.Data && InInputAccessUnit.DataSize)
 	{
 		// Parse the codec specific information
-		if (!bHaveParseddOps && !Parse_dOps())
+		if (!bHaveParseddOps)
 		{
-			return IElectraDecoder::EDecoderError::Error;
+			if (!Parse_dOps(DOpsConfig, ElectraDecodersUtil::GetVariantValueUInt8Array(InAdditionalOptions, TEXT("$dOps_box")), true))
+			{
+				return IElectraDecoder::EDecoderError::Error;
+			}
+			bHaveParseddOps = true;
 		}
 		// Set up the channel map accordingly.
 		if (!SetupChannelMap())
