OpusDecoderElectra: remove application of preskip since this is handled in the upper decoder layer and would otherwise be applied twice.

#rnx

diff --git a/Engine/Plugins/Media/OpusDecoderElectra/Source/OpusDecoderElectra/Private/OpusDecoder/ElectraMediaOpusDecoder.cpp b/Engine/Plugins/Media/OpusDecoderElectra/Source/OpusDecoderElectra/Private/OpusDecoder/ElectraMediaOpusDecoder.cpp
--- a/Engine/Plugins/Media/OpusDecoderElectra/Source/OpusDecoderElectra/Private/OpusDecoder/ElectraMediaOpusDecoder.cpp
+++ b/Engine/Plugins/Media/OpusDecoderElectra/Source/OpusDecoderElectra/Private/OpusDecoder/ElectraMediaOpusDecoder.cpp
@@ -653,6 +653,11 @@ bool FElectraOpusDecoder::ProcessInput(const void* InData, int64 InDataSize)
 	}
 	CurrentOutput->NumFrames = Result;
 
+#if 0
+	// Do not apply the pre-skip here. The upper layer decoder is responsible for adjusting the
+	// decoded samples based on the PTS value, which will be negative at the start to account
+	// for the duration of the pre-skip.
+
 	// Apply pre skip
 	if (RemainingPreSkip > 0)
 	{
@@ -674,6 +679,7 @@ bool FElectraOpusDecoder::ProcessInput(const void* InData, int64 InDataSize)
 			RemainingPreSkip = 0;
 		}
 	}
+#endif
 	return true;
 }
 
