ElectraDecoder: Opus decoder needs to be configured with the next highest Opus-native sample rate if the original is not an exact match

#rnx

diff --git a/Engine/Plugins/Media/OpusDecoderElectra/Source/OpusDecoderElectra/Private/OpusDecoder/ElectraMediaOpusDecoder.cpp b/Engine/Plugins/Media/OpusDecoderElectra/Source/OpusDecoderElectra/Private/OpusDecoder/ElectraMediaOpusDecoder.cpp
--- a/Engine/Plugins/Media/OpusDecoderElectra/Source/OpusDecoderElectra/Private/OpusDecoder/ElectraMediaOpusDecoder.cpp
+++ b/Engine/Plugins/Media/OpusDecoderElectra/Source/OpusDecoderElectra/Private/OpusDecoder/ElectraMediaOpusDecoder.cpp
@@ -198,7 +198,23 @@ private:
 
 	int32 OpusSamplingRate() const
 	{
-		return DOpsConfig.SampleRate;
+		if (DOpsConfig.SampleRate <= 0 || DOpsConfig.SampleRate > 24000)
+		{
+			return 48000;
+		}
+		else if (DOpsConfig.SampleRate > 16000)
+		{
+			return 24000;
+		}
+		else if (DOpsConfig.SampleRate > 12000)
+		{
+			return 16000;
+		}
+		else if (DOpsConfig.SampleRate > 8000)
+		{
+			return 12000;
+		}
+		return 8000;
 	}
 
 	bool Parse_dOps(FDOpsConfig& OutConfig, const TArray<uint8>& IndOpsBox, bool bFailOnError);
